Adding SlateScripting plugin, SlateScriptingCommands module and UICommandsScriptingSubsystem to register commands dynamically within existing UI Command LIsts through scripting. Sequencer, Curve Editor, Content Browser, Main Frame, Level Editor and Level Viewports are currently the only systems exposing the command lists.

diff --git a/Engine/Plugins/Slate/SlateScripting/SlateScripting.uplugin b/Engine/Plugins/Slate/SlateScripting/SlateScripting.uplugin
new file mode 100644
--- /dev/null
+++ b/Engine/Plugins/Slate/SlateScripting/SlateScripting.uplugin
@@ -0,0 +1,24 @@
+{
+	"FileVersion": 3,
+	"Version": 1,
+	"VersionName": "1.0",
+	"FriendlyName": "Slate Scripting",
+	"Description": "Allows interacting with Slate through scripting",
+	"Category": "Scripting",
+	"CreatedBy": "Epic Games, Inc.",
+	"CreatedByURL": "http://epicgames.com",
+	"DocsURL": "",
+	"MarketplaceURL": "",
+	"SupportURL": "",
+	"EnabledByDefault": false,
+	"CanContainContent": false,
+	"IsBetaVersion": false,
+	"Installed": false,
+	"Modules": [
+		{
+			"Name": "SlateScriptingCommands",
+			"Type": "Runtime",
+			"LoadingPhase": "Default"
+		}
+	]
+}
diff --git a/Engine/Plugins/Slate/SlateScripting/Source/SlateScriptingCommands/Private/SlateScriptingCommandsModule.cpp b/Engine/Plugins/Slate/SlateScripting/Source/SlateScriptingCommands/Private/SlateScriptingCommandsModule.cpp
new file mode 100644
--- /dev/null
+++ b/Engine/Plugins/Slate/SlateScripting/Source/SlateScriptingCommands/Private/SlateScriptingCommandsModule.cpp
@@ -0,0 +1,13 @@
+// Copyright Epic Games, Inc. All Rights Reserved.
+
+#include "SlateScriptingCommandsModule.h"
+
+void FSlateScriptingCommandsModule::StartupModule()
+{
+}
+
+void FSlateScriptingCommandsModule::ShutdownModule()
+{
+}
+
+IMPLEMENT_MODULE(FSlateScriptingCommandsModule, SlateScriptingCommands);
\ No newline at end of file
diff --git a/Engine/Plugins/Slate/SlateScripting/Source/SlateScriptingCommands/Private/UICommandsScriptingSubsystem.cpp b/Engine/Plugins/Slate/SlateScripting/Source/SlateScriptingCommands/Private/UICommandsScriptingSubsystem.cpp
new file mode 100644
--- /dev/null
+++ b/Engine/Plugins/Slate/SlateScripting/Source/SlateScriptingCommands/Private/UICommandsScriptingSubsystem.cpp
@@ -0,0 +1,513 @@
+// Copyright Epic Games, Inc. All Rights Reserved.
+
+#include "UICommandsScriptingSubsystem.h"
+
+#include "Framework/Commands/UIAction.h"
+#include "Framework/Commands/InputBindingManager.h"
+
+DEFINE_LOG_CATEGORY(LogCommandsScripting)
+
+
+bool FScriptingCommandInfo::Equals(const FScriptingCommandInfo& InCommandInfo, const bool bCheckInputChord) const
+{
+	return InCommandInfo.Name == Name && InCommandInfo.Set == Set && InCommandInfo.ContextName == ContextName &&
+		(!bCheckInputChord || InCommandInfo.InputChord == InputChord);
+}
+
+TSharedPtr<FUICommandInfo> FScriptingCommand::MakeUICommandInfo() const
+{
+	FInputBindingManager& BindingManager = FInputBindingManager::Get();
+	const FName CommandName = GetFullName();
+	const TSharedPtr<FBindingContext> Context = BindingManager.GetContextByName(CommandInfo.ContextName);
+
+	if (!Context.IsValid())
+	{
+		UE_LOG(LogCommandsScripting, Error, TEXT("Context is not registered in the Input Binding Manager: %s"), *CommandInfo.ContextName.ToString())
+		return nullptr;
+	}
+	if (BindingManager.FindCommandInContext(CommandInfo.ContextName, CommandName))
+	{
+		UE_LOG(LogCommandsScripting, Error, TEXT("%s: Command already registered in context %s"), *CommandName.ToString(), *CommandInfo.ContextName.ToString())
+		return nullptr;
+	}
+	if (BindingManager.FindCommandInContext(CommandInfo.ContextName, CommandInfo.InputChord, false))
+	{
+		UE_LOG(LogCommandsScripting, Error, TEXT("Input Chord already mapped in context: %s"), *CommandInfo.ContextName.ToString())
+		return nullptr;
+	}
+
+	TSharedPtr<FUICommandInfo> NewCommand;
+
+	FUICommandInfo::MakeCommandInfo(
+		Context.ToSharedRef(),
+		NewCommand,
+		CommandName,
+		CommandInfo.Label,
+		CommandInfo.Description,
+		FSlateIcon(),
+		EUserInterfaceActionType::Button,
+		CommandInfo.InputChord
+	);
+
+	return NewCommand;
+}
+
+bool FScriptingCommand::UnregisterUICommandInfo() const
+{
+	FInputBindingManager& BindingManager = FInputBindingManager::Get();
+	const TSharedPtr<FBindingContext> Context = BindingManager.GetContextByName(CommandInfo.ContextName);
+	const TSharedPtr<FUICommandInfo> Command = BindingManager.FindCommandInContext(CommandInfo.ContextName, GetFullName());
+
+	if (Context.IsValid() && Command.IsValid())
+	{
+		FUICommandInfo::UnregisterCommandInfo(Context.ToSharedRef(), Command.ToSharedRef());
+		return true;
+	}
+	return false;
+}
+
+
+
+bool FScriptingCommandsContext::MapCommand(const TSharedRef<FScriptingCommand> ScriptingCommand)
+{
+	const TSharedPtr<FUICommandInfo> Command = FInputBindingManager::Get().FindCommandInContext(
+		ScriptingCommand->CommandInfo.ContextName, ScriptingCommand->GetFullName());
+
+	if (Command.IsValid())
+	{
+		CleanupPointerArray(CommandLists);
+
+		for (TWeakPtr<FUICommandList> CommandList : CommandLists)
+		{
+			CommandList.Pin()->MapAction(Command, ScriptingCommand->OnExecuteAction, ScriptingCommand->OnCanExecuteAction);
+		}
+	
+		ScriptingCommands.AddUnique(ScriptingCommand);
+		return true;
+	}
+	
+	UE_LOG(LogCommandsScripting, Error, TEXT("Could not map command: %s. The command could not be found in the Input Binding Manager"), *ScriptingCommand->GetFullName().ToString())
+	return false;
+}
+
+bool FScriptingCommandsContext::UnmapCommand(const TSharedRef<FScriptingCommand> ScriptingCommand)
+{
+	const TSharedPtr<FUICommandInfo> Command = FInputBindingManager::Get().FindCommandInContext(
+		ScriptingCommand->CommandInfo.ContextName, ScriptingCommand->GetFullName());
+
+	if (Command.IsValid())
+	{
+		CleanupPointerArray(CommandLists);
+	
+		for (TWeakPtr<FUICommandList> CommandList : CommandLists)
+		{
+			CommandList.Pin()->UnmapAction(Command);
+		}
+
+		ScriptingCommands.Remove(ScriptingCommand);
+		return true;
+	}
+	
+	UE_LOG(LogCommandsScripting, Error, TEXT("Could not unmap command: %s. The command could not be found in the Input Binding Manager"), *ScriptingCommand->GetFullName().ToString())
+	return false;
+}
+
+bool FScriptingCommandsContext::RegisterCommandList(const TSharedRef<FUICommandList> CommandList)
+{
+	CleanupPointerArray(CommandLists);
+
+	if (!CommandLists.Contains(CommandList))
+	{
+		CommandLists.Add(CommandList);
+		MapAllCommands(CommandList);
+		return true;
+	}
+	
+	UE_LOG(LogCommandsScripting, Warning, TEXT("Trying to register an already registered command list"))
+	return false;
+}
+
+bool FScriptingCommandsContext::UnregisterCommandList(const TSharedRef<FUICommandList> CommandList)
+{
+	CleanupPointerArray(CommandLists);
+	
+	if (CommandLists.Remove(CommandList) > 0)
+	{
+		UnmapAllCommands(CommandList);
+		return true;
+	}
+	
+	UE_LOG(LogCommandsScripting, Warning, TEXT("Trying to unregister a non registered command list"))
+	return false;
+}
+
+void FScriptingCommandsContext::MapAllCommands(const TSharedRef<FUICommandList> CommandList)
+{
+	for (const TSharedPtr<FScriptingCommand>& ScriptingCommand : ScriptingCommands)
+	{
+		if (ScriptingCommand.IsValid())
+		{
+			const TSharedPtr<FUICommandInfo> CommandInfo = FInputBindingManager::Get().FindCommandInContext(
+				ContextName, ScriptingCommand->GetFullName());
+
+			if (CommandInfo.IsValid())
+			{
+				CommandList->MapAction(CommandInfo, ScriptingCommand->OnExecuteAction, ScriptingCommand->OnCanExecuteAction);
+			}	
+		}
+	}
+}
+
+void FScriptingCommandsContext::UnmapAllCommands(const TSharedRef<FUICommandList> CommandList)
+{
+	for (const TSharedPtr<FScriptingCommand>& ScriptingCommand : ScriptingCommands)
+	{
+		if (ScriptingCommand.IsValid())
+		{
+			const TSharedPtr<FUICommandInfo> CommandInfo = FInputBindingManager::Get().FindCommandInContext(
+				ContextName, ScriptingCommand->GetFullName());
+
+			if (CommandInfo.IsValid())
+			{
+				CommandList->UnmapAction(CommandInfo);
+			}
+		}
+	}
+}
+
+
+void UUICommandsScriptingSubsystem::Initialize(FSubsystemCollectionBase& Collection)
+{
+	FInputBindingManager::Get().OnRegisterCommandList.AddUObject(this, &UUICommandsScriptingSubsystem::RegisterCommandListForContext);
+	FInputBindingManager::Get().OnUnregisterCommandList.AddUObject(this, &UUICommandsScriptingSubsystem::UnregisterCommandListForContext);
+}
+
+void UUICommandsScriptingSubsystem::Deinitialize()
+{
+	UnregisterAllSets();
+}
+
+void UUICommandsScriptingSubsystem::RegisterCommandListForContext(const FName ContextName,
+                                                          const TSharedRef<FUICommandList> CommandList)
+{
+	FScriptingCommandsContext& CommandsContext = CommandsInContext.FindOrAdd(ContextName, ContextName);
+	CommandsContext.RegisterCommandList(CommandList);
+}
+
+void UUICommandsScriptingSubsystem::UnregisterCommandListForContext(const FName ContextName,
+                                                            const TSharedRef<FUICommandList> CommandList)
+{
+	FScriptingCommandsContext *CommandsContext = CommandsInContext.Find(ContextName);
+
+	if (CommandsContext)
+	{
+		CommandsContext->UnregisterCommandList(CommandList);
+	}
+}
+
+bool UUICommandsScriptingSubsystem::UnregisterContext(const FName ContextName)
+{
+	FScriptingCommandsContext *CommandsContext = CommandsInContext.Find(ContextName);
+
+	if (CommandsContext)
+	{
+		CleanupPointerArray(CommandsContext->CommandLists);
+
+		for (const TWeakPtr<FUICommandList>& CommandList : CommandsContext->CommandLists)
+		{
+			CommandsContext->UnmapAllCommands(CommandList.Pin().ToSharedRef());
+		}
+		
+		CommandsInContext.Remove(ContextName);
+		return true;
+	}
+	
+	UE_LOG(LogCommandsScripting, Warning, TEXT("Trying to unregister a non registered context"))
+	return false;
+}
+
+
+
+bool UUICommandsScriptingSubsystem::RegisterCommand(const FScriptingCommandInfo CommandInfo,
+                                            const FExecuteCommand OnExecuteCommand,
+                                            const bool bOverrideExisting)
+{
+	// Registers the command with default CanExecute (just checking if command set and subsystem are enabled).
+	const FCanExecuteAction OnCanExecuteAction = FCanExecuteAction::CreateUObject(this,
+		&UUICommandsScriptingSubsystem::DefaultCanExecuteAction, CommandInfo.Set);
+
+	return RegisterNewScriptingCommand(CommandInfo, OnExecuteCommand, OnCanExecuteAction, bOverrideExisting);
+}
+
+bool UUICommandsScriptingSubsystem::RegisterCommandChecked(const FScriptingCommandInfo CommandInfo,
+                                                   const FExecuteCommand OnExecuteCommand,
+                                                   const FCanExecuteCommand OnCanExecuteCommand,
+                                                   const bool bOverrideExisting)
+{
+	// Registers the command with the given CanExecute delegate (still checking if command set and subsystem are enabled)
+	const FCanExecuteAction OnCanExecuteAction = FCanExecuteAction::CreateUObject(this,
+		&UUICommandsScriptingSubsystem::HandleCanExecuteAction, OnCanExecuteCommand, CommandInfo);
+
+	return RegisterNewScriptingCommand(CommandInfo, OnExecuteCommand, OnCanExecuteAction, bOverrideExisting);
+}
+
+bool UUICommandsScriptingSubsystem::UnregisterCommand(FScriptingCommandInfo CommandInfo)
+{
+	if (FScriptingCommandSet* CommandSet = CommandSets.Find(CommandInfo.Set))
+	{
+		TSharedPtr<FScriptingCommand>* ScriptingCommand = CommandSet->ScriptingCommands.FindByPredicate([CommandInfo](
+			const TSharedPtr<FScriptingCommand> Command) { return Command.IsValid() && Command->CommandInfo.Equals(CommandInfo);});
+		if (ScriptingCommand && ScriptingCommand->IsValid())
+		{
+			UnregisterScriptingCommand(ScriptingCommand->ToSharedRef());
+			CommandSet->ScriptingCommands.RemoveSingle(*ScriptingCommand);
+			return true;
+		}
+		
+		UE_LOG(LogCommandsScripting, Error, TEXT("Command not registered: %s"), *CommandInfo.GetFullName().ToString())
+		return false;
+	}
+	
+	UE_LOG(LogCommandsScripting, Error, TEXT("Command Set not registered: %s"), *CommandInfo.Set.ToString())
+	return false;
+}
+
+bool UUICommandsScriptingSubsystem::UnregisterCommandSet(const FName SetName)
+{
+	if (const FScriptingCommandSet* CommandSet = CommandSets.Find(SetName))
+	{
+		for (const TSharedPtr<FScriptingCommand>& ScriptingCommand : CommandSet->ScriptingCommands)
+		{
+			if (ScriptingCommand.IsValid())
+			{
+				UnregisterScriptingCommand(ScriptingCommand.ToSharedRef());
+			}
+		}
+		CommandSets.Remove(SetName);
+		return true;
+	}
+	
+	UE_LOG(LogCommandsScripting, Error, TEXT("Command Set not registered: %s"), *SetName.ToString())
+	return false;
+}
+
+void UUICommandsScriptingSubsystem::UnregisterAllSets()
+{
+	for (const TPair<FName, FScriptingCommandSet>& CommandSet : CommandSets)
+	{
+		for (const TSharedPtr<FScriptingCommand>& ScriptingCommand : CommandSet.Value.ScriptingCommands)
+		{
+			if (ScriptingCommand.IsValid())
+			{
+				UnregisterScriptingCommand(ScriptingCommand.ToSharedRef());
+			}
+		}
+	}
+	CommandSets.Reset();
+}
+
+TArray<FName> UUICommandsScriptingSubsystem::GetAvailableContexts() const
+{
+	TArray<FName> OutContexts;
+
+	CommandsInContext.GetKeys(OutContexts);
+	return OutContexts;
+}
+
+bool UUICommandsScriptingSubsystem::IsContextRegistered(const FName ContextName) const
+{
+	return CommandsInContext.Contains(ContextName);
+}
+
+int UUICommandsScriptingSubsystem::GetBindingCountForContext(const FName ContextName)
+{
+	if (IsContextRegistered(ContextName))
+	{
+		CleanupPointerArray(CommandsInContext[ContextName].CommandLists);
+		return CommandsInContext[ContextName].CommandLists.Num();
+	}
+	
+	UE_LOG(LogCommandsScripting, Error, TEXT("Context not registered: %s"), *ContextName.ToString())
+	return 0;
+}
+
+TArray<FScriptingCommandInfo> UUICommandsScriptingSubsystem::GetRegisteredCommands() const
+{
+	TArray<FScriptingCommandInfo> OutCommandsInfo;
+
+	for (const TPair<FName, FScriptingCommandSet>& CommandSet : CommandSets)
+	{
+		for (const TSharedPtr<FScriptingCommand>& ScriptingCommand : CommandSet.Value.ScriptingCommands)
+		{
+			if (ScriptingCommand.IsValid())
+			{
+				OutCommandsInfo.Add(ScriptingCommand->CommandInfo);
+			}
+		}
+	}
+	return OutCommandsInfo;
+}
+
+bool UUICommandsScriptingSubsystem::IsCommandRegistered(const FScriptingCommandInfo CommandInfo, const bool bCheckInputChord) const
+{
+	if (CommandSets.Contains(CommandInfo.Set))
+	{
+		const TSharedPtr<FScriptingCommand>* ExistingCommand = CommandSets[CommandInfo.Set].ScriptingCommands.FindByPredicate(
+			[CommandInfo, bCheckInputChord](const TSharedPtr<FScriptingCommand> ScriptingCommand)
+		{
+			return ScriptingCommand.IsValid() && ScriptingCommand->CommandInfo.Equals(CommandInfo, bCheckInputChord);
+		});
+		
+		return ExistingCommand && ExistingCommand->IsValid();
+	}
+	return false;
+}
+
+auto UUICommandsScriptingSubsystem::IsInputChordMapped(const FName ContextName, const FInputChord InputChord) const -> bool
+{
+	if (!IsContextRegistered(ContextName))
+	{
+		UE_LOG(LogCommandsScripting, Error, TEXT("Context not registered: %s"), *ContextName.ToString())
+		return false;
+	}
+	
+	return FInputBindingManager::Get().FindCommandInContext(ContextName, InputChord, false) != nullptr;
+}
+
+
+void UUICommandsScriptingSubsystem::SetCanSetExecuteCommands(const FName SetName, const bool bShouldExecuteCommands)
+{
+	if (FScriptingCommandSet* CommandSet = CommandSets.Find(SetName))
+	{
+		CommandSet->bCanExecuteCommands = bShouldExecuteCommands;
+	}
+	else
+	{
+		UE_LOG(LogCommandsScripting, Error, TEXT("Command Set not registered: %s"), *SetName.ToString())
+	}
+}
+
+bool UUICommandsScriptingSubsystem::CanSetExecuteCommands(const FName SetName) const
+{
+	if (const FScriptingCommandSet* CommandSet = CommandSets.Find(SetName))
+	{
+		return bCanExecuteCommands && CommandSet->bCanExecuteCommands;
+	}
+	
+	UE_LOG(LogCommandsScripting, Error, TEXT("Command Set not registered: %s"), *SetName.ToString())
+	return false;
+}
+
+void UUICommandsScriptingSubsystem::SetCanExecuteCommands(const bool bShouldExecuteCommands)
+{
+	bCanExecuteCommands = bShouldExecuteCommands;
+}
+
+bool UUICommandsScriptingSubsystem::IsCommandSetRegistered(const FName SetName) const
+{
+	return CommandSets.Contains(SetName);
+}
+
+bool UUICommandsScriptingSubsystem::RegisterCommandSet(const FName SetName)
+{
+	if (IsCommandSetRegistered(SetName))
+	{
+		UE_LOG(LogCommandsScripting, Warning, TEXT("Command Set already registered: %s"), *SetName.ToString())
+		return false;
+	}
+
+	CommandSets.Add(SetName);
+	return true;
+}
+
+bool UUICommandsScriptingSubsystem::CanExecuteCommands() const
+{
+	return bCanExecuteCommands;
+}
+
+
+
+bool UUICommandsScriptingSubsystem::RegisterNewScriptingCommand(const FScriptingCommandInfo CommandInfo,
+                                                      const FExecuteCommand OnExecuteCommand,
+                                                      const FCanExecuteAction OnCanExecuteAction,
+                                                      const bool bOverrideExisting)
+{
+	if (!IsCommandSetRegistered(CommandInfo.Set))
+	{
+		UE_LOG(LogCommandsScripting, Error, TEXT("Trying to register a command in a non registered Command Set: %s"), *CommandInfo.Set.ToString())
+		return false;
+	}
+	
+	// Bind the given OnExecute delegate to the non dynamic one expected internally
+	const FExecuteAction OnExecuteAction = FExecuteAction::CreateStatic(
+		&UUICommandsScriptingSubsystem::HandleExecuteAction, OnExecuteCommand, CommandInfo);
+
+	// Create the command data to be cached by the subsystem
+	const TSharedPtr<FScriptingCommand> ScriptingCommand = MakeShareable(new FScriptingCommand(CommandInfo, OnExecuteAction, OnCanExecuteAction));
+
+	// Unregister any command previously registered with this name to overwrite it
+	if (IsCommandRegistered(CommandInfo, false))
+	{
+		if (!bOverrideExisting)
+		{
+			UE_LOG(LogCommandsScripting, Warning, TEXT("Previously registered command %s won't be overriden"), *CommandInfo.GetFullName().ToString())
+			return false;
+		}
+		if (!UnregisterCommand(CommandInfo))
+		{
+			UE_LOG(LogCommandsScripting, Error, TEXT("Could not override command: %s"), *CommandInfo.GetFullName().ToString())
+			return false;
+		}
+	}
+	
+	// Register the command without checking if it has been bound to a command list (if it wasn't, it will when possible)
+	if (RegisterScriptingCommand(ScriptingCommand.ToSharedRef()))
+	{
+		// Add the command to our list of registered commands
+		CommandSets[CommandInfo.Set].ScriptingCommands.AddUnique(ScriptingCommand);
+		return true;
+	}
+	return false;
+}
+
+bool UUICommandsScriptingSubsystem::RegisterScriptingCommand(const TSharedRef<FScriptingCommand> ScriptingCommand)
+{
+	FScriptingCommandsContext* CommandsContext = CommandsInContext.Find(ScriptingCommand->CommandInfo.ContextName);
+
+	if (CommandsContext)
+	{
+		return ScriptingCommand->MakeUICommandInfo() && CommandsContext->MapCommand(ScriptingCommand);
+	}
+
+	UE_LOG(LogCommandsScripting, Error, TEXT("Context not registered: %s"), *ScriptingCommand->CommandInfo.ContextName.ToString())
+	return false;
+}
+
+bool UUICommandsScriptingSubsystem::UnregisterScriptingCommand(const TSharedRef<FScriptingCommand> ScriptingCommand)
+{
+	FScriptingCommandsContext* CommandsContext = CommandsInContext.Find(ScriptingCommand->CommandInfo.ContextName);
+
+	if (CommandsContext)
+	{
+		return CommandsContext->UnmapCommand(ScriptingCommand) && ScriptingCommand->UnregisterUICommandInfo();
+	}
+
+	UE_LOG(LogCommandsScripting, Error, TEXT("Context not registered: %s"), *ScriptingCommand->CommandInfo.ContextName.ToString())
+	return false;
+}
+
+
+void UUICommandsScriptingSubsystem::HandleExecuteAction(const FExecuteCommand OnExecuteAction, const FScriptingCommandInfo CommandInfo)
+{
+	OnExecuteAction.Execute(CommandInfo);
+}
+
+bool UUICommandsScriptingSubsystem::HandleCanExecuteAction(const FCanExecuteCommand OnCanExecuteAction, const FScriptingCommandInfo CommandInfo) const
+{
+	return CanSetExecuteCommands(CommandInfo.Set) && OnCanExecuteAction.Execute(CommandInfo);
+}
+
+bool UUICommandsScriptingSubsystem::DefaultCanExecuteAction(const FName SetName) const
+{
+	return CanSetExecuteCommands(SetName);
+}
diff --git a/Engine/Plugins/Slate/SlateScripting/Source/SlateScriptingCommands/Public/SlateScriptingCommandsModule.h b/Engine/Plugins/Slate/SlateScripting/Source/SlateScriptingCommands/Public/SlateScriptingCommandsModule.h
new file mode 100644
--- /dev/null
+++ b/Engine/Plugins/Slate/SlateScripting/Source/SlateScriptingCommands/Public/SlateScriptingCommandsModule.h
@@ -0,0 +1,15 @@
+// Copyright Epic Games, Inc. All Rights Reserved.
+
+#pragma once
+
+#include "CoreMinimal.h"
+#include "Modules/ModuleInterface.h"
+
+class FSlateScriptingCommandsModule : public IModuleInterface
+{
+public:
+	//~ Begin IModuleInterface interface
+	virtual void StartupModule() override;
+	virtual void ShutdownModule() override;
+	//~ End IModuleInterface interface
+};
diff --git a/Engine/Plugins/Slate/SlateScripting/Source/SlateScriptingCommands/Public/UICommandsScriptingSubsystem.h b/Engine/Plugins/Slate/SlateScripting/Source/SlateScriptingCommands/Public/UICommandsScriptingSubsystem.h
new file mode 100644
--- /dev/null
+++ b/Engine/Plugins/Slate/SlateScripting/Source/SlateScriptingCommands/Public/UICommandsScriptingSubsystem.h
@@ -0,0 +1,368 @@
+// Copyright Epic Games, Inc. All Rights Reserved.
+
+#pragma once
+
+#include "CoreMinimal.h"
+
+#include "Subsystems/EngineSubsystem.h"
+
+#include "Framework/Commands/UICommandList.h"
+#include "Framework/Commands/UICommandInfo.h"
+#include "Framework/Commands/InputChord.h"
+
+#include "UICommandsScriptingSubsystem.generated.h"
+
+
+DECLARE_LOG_CATEGORY_EXTERN(LogCommandsScripting, Log, All);
+
+
+/**
+ * The data defining a scripting command. At the exception of its delegates.
+ */
+USTRUCT(BlueprintType)
+struct FScriptingCommandInfo
+{
+	GENERATED_BODY()
+
+	/** The editor context this command is bound to */
+	UPROPERTY(EditAnywhere, BlueprintReadWrite, Category = "Editor Scripting | Commands")
+	FName ContextName;
+
+	/** The command set this command belongs to. This is to avoid conflicts and could refer to the owner of the command */
+	UPROPERTY(EditAnywhere, BlueprintReadWrite, Category = "Editor Scripting | Commands")
+	FName Set;
+
+	/** The command name. Must be unique in its set. */
+	UPROPERTY(EditAnywhere, BlueprintReadWrite, Category = "Editor Scripting | Commands")
+	FName Name;
+
+	/** The command label or what name will be displayed for it */
+	UPROPERTY(EditAnywhere, BlueprintReadWrite, Category = "Editor Scripting | Commands")
+	FText Label;
+
+	/** The description of the command */
+	UPROPERTY(EditAnywhere, BlueprintReadWrite, Category = "Editor Scripting | Commands")
+	FText Description;
+
+	/** The input chord to bound to the command */
+	UPROPERTY(EditAnywhere, BlueprintReadWrite, Category = "Editor Scripting | Commands")
+	FInputChord InputChord;
+
+	/** Compares this command info with the given one per context, set, name and optionally input chord */
+	bool Equals(const FScriptingCommandInfo& InCommandInfo, const bool bCheckInputChord = false) const;
+
+	/** Builds a full name in the format context.set.command_name to avoid conflicts with commands registered in different contexts/sets */
+	FName GetFullName() const
+	{
+		// We append the set name and command name to avoid conflicts between commands registered by different scripts/components
+		return *(ContextName.ToString() + "." + Set.ToString() + "." + Name.ToString());
+	}
+};
+
+/** Exposing FExecuteAction as dynamic */
+DECLARE_DYNAMIC_DELEGATE_OneParam(FExecuteCommand, FScriptingCommandInfo, CommandInfo);
+
+/** Exposing FCanExecuteAction as dynamic */
+DECLARE_DYNAMIC_DELEGATE_RetVal_OneParam(bool, FCanExecuteCommand, FScriptingCommandInfo, CommandInfo);
+
+/**
+ * All the internal data related to a scripting command as well as methods to easily expose the command to the binding manager
+ */
+USTRUCT()
+struct FScriptingCommand
+{
+	GENERATED_BODY()
+
+	FScriptingCommand() = default;
+	
+	FScriptingCommand(const FScriptingCommandInfo& InCommandInfo, const FExecuteAction InOnExecuteAction,
+	                const FCanExecuteAction InOnCanExecuteAction) : CommandInfo(InCommandInfo),
+																	OnExecuteAction(InOnExecuteAction),
+																	OnCanExecuteAction(InOnCanExecuteAction)
+	{}
+
+	/** The command definition */
+	FScriptingCommandInfo CommandInfo;
+
+	/** The delegate to call when executing the command */
+	FExecuteAction OnExecuteAction;
+
+	/** The delegate to call to check whether the command can be executed */
+	FCanExecuteAction OnCanExecuteAction;
+
+	FORCEINLINE bool operator==(const FScriptingCommand& Other) const { return Other.CommandInfo.Name == CommandInfo.Name; }
+	FORCEINLINE bool operator!=(const FScriptingCommand& Other) const { return !(*this == Other); }
+
+	FORCEINLINE bool operator==(const FName OtherName) const { return OtherName == CommandInfo.Name; }
+	FORCEINLINE bool operator!=(const FName OtherName) const { return !(*this == OtherName); }
+	
+	/** Builds a full name in the format context.set.command_name to avoid conflicts with commands registered in different contexts/sets  */
+	FName GetFullName() const
+	{
+		return CommandInfo.GetFullName();
+	}
+
+	/**
+	 * Registers a new UICommandInfo in the associated context through the binding manager.
+	 * At this point, the command is exposed to the manager but not mapped to any command list.
+	 */
+	TSharedPtr<FUICommandInfo> MakeUICommandInfo() const;
+
+	/** Unregisters this command's UICommandInfo from its associated context through the binding manager */
+	bool UnregisterUICommandInfo() const;
+};
+
+/**
+ * The list of commands and UI Command Lists associated with a context.
+ * This enables easier management of commands within registered contexts and their UI Command Lists.
+ */
+USTRUCT()
+struct FScriptingCommandsContext
+{
+	GENERATED_BODY()
+
+	FScriptingCommandsContext() = default;
+
+	FScriptingCommandsContext(const FName InContextName) : ContextName(InContextName) {}
+
+	/** The context these command lists are bound to */
+	FName ContextName;
+	
+	/** An array of command list associated to the given context */
+	TArray<TWeakPtr<FUICommandList>> CommandLists;
+
+	/** The commands bound in these command lists */
+	TArray<TSharedPtr<FScriptingCommand>> ScriptingCommands;
+
+public:
+	/** Maps the given command to all the command lists of this context */
+	bool MapCommand(const TSharedRef<FScriptingCommand> ScriptingCommand);
+
+	/** Unmaps the given command from all the command lists of this context */
+	bool UnmapCommand(const TSharedRef<FScriptingCommand> ScriptingCommand);
+
+	/** Registers a command list in this context then map to it all commands existing within this context */
+	bool RegisterCommandList(const TSharedRef<FUICommandList> CommandList);
+
+	/** Unregisters command list from this context then unmap from it all commands existing within this context */
+	bool UnregisterCommandList(const TSharedRef<FUICommandList> CommandList);
+
+	/** Maps all commands existing within this context to the given command list */
+	void MapAllCommands(TSharedRef<FUICommandList> CommandList);
+
+	/** Unmaps all commands existing within this context from the given command list */
+	void UnmapAllCommands(TSharedRef<FUICommandList> CommandList);
+};
+
+
+/**
+ * To avoid conflicts between scripts, each command is associated to a command set which must be registered manually.
+ */
+USTRUCT()
+struct FScriptingCommandSet
+{
+	GENERATED_BODY()
+
+	/** The commands in this set */
+	TArray<TSharedPtr<FScriptingCommand>> ScriptingCommands;
+	
+	/** Whether the commands in this set are enabled */
+	bool bCanExecuteCommands = true;
+};
+
+/**
+ * UEditorInputSubsystem
+ * Subsystem for dynamically registering editor commands through scripting
+ */
+UCLASS()
+class SLATESCRIPTINGCOMMANDS_API UUICommandsScriptingSubsystem : public UEngineSubsystem
+{
+	GENERATED_BODY()
+
+public:
+	/**
+	 * USubsystem Interface: handles initialization of instances of the system
+	 * This is where the subsystem binds to the CommandListRegistered delegate of the InputBindingManager
+	 *
+	 * To expose command lists from external systems to this subsystem, call FInputBindingManager::RegisterCommandList
+	 * from the external system. The binding manager will then broadcast the registered command list to all subscribers.
+	 */
+	virtual void Initialize(FSubsystemCollectionBase& Collection) override;
+	
+	/** USubsystem Interface: handles deinitialization of instances of the system */
+	virtual void Deinitialize() override;
+
+	/**
+	 * Registers a command list in the given context. Optionally registering the context itself if it does not exist
+	 * Note the context should already be existing and registered through the binding manager.
+	 * 
+	 * This should be called by modules and tools looking to expose there command lists to the subsystem
+	 * 
+	 * @param	ContextName		The context name as defined in the binding manager
+	 * @param	CommandList		The command list to register in this context
+	 */
+	void RegisterCommandListForContext(const FName ContextName, TSharedRef<FUICommandList> CommandList);
+
+	/**
+	 * Unregisters a command list from the given context.
+	 * @param	ContextName		The context name as defined in the binding manager
+	 * @param	CommandList		The command list to register in this context
+	 */
+	void UnregisterCommandListForContext(const FName ContextName, TSharedRef<FUICommandList> CommandList);
+
+	/**
+	 * Unregisters all command lists from the given context.
+	 * @param	ContextName		The context name as defined in the binding manager
+	 */
+	bool UnregisterContext(const FName ContextName);
+
+public:	
+	/**
+	 * Registers a command within the given context and set.
+	 * The set must be registered beforehand.
+	 * @param	CommandInfo			The command infos such as name, label, description and input chord.
+	 * @param	OnExecuteCommand	The delegate to be executed for handling this command.
+	 * @param	bOverrideExisting	Whether existing command with matching context, set and name should be overriden
+	 * @return Whether the command was succesfully registered
+	 */
+	UFUNCTION(BlueprintCallable, Category = "Editor Scripting | Commands |")
+	bool RegisterCommand(FScriptingCommandInfo CommandInfo, FExecuteCommand OnExecuteCommand, bool bOverrideExisting = false);
+
+	/**
+	 * Registers a command within the given context and set.
+	 * The set must be registered beforehand.
+	 * @param	CommandInfo				The command infos such as name, label, description and input chord.
+	 * @param	OnExecuteCommand 		The delegate to be executed for handling this command.
+	 * @param	OnCanExecuteCommand 	The delegate to be executed for checking if this command can be executed.
+	 * @param	bOverrideExisting		Whether existing command with matching context, set and name should be overriden
+	 * @return Whether the command was successfully registered
+	 */
+	UFUNCTION(BlueprintCallable, Category = "Editor Scripting | Commands |")
+	bool RegisterCommandChecked(FScriptingCommandInfo CommandInfo, FExecuteCommand OnExecuteCommand,
+	                            FCanExecuteCommand OnCanExecuteCommand, bool bOverrideExisting = false);
+
+	/**
+	 * Unregisters a command previously registered. The command name, set and context will be used for comparison.
+	 * @return Whether the command was successfully unregistered.
+	 */
+	UFUNCTION(BlueprintCallable, Category = "Editor Scripting | Commands |")
+	bool UnregisterCommand(FScriptingCommandInfo CommandInfo);
+
+	/** Checks whether the given command is registered within the subsystem. Using name, set and context for comparison */
+	UFUNCTION(BlueprintCallable, Category = "Editor Scripting | Commands |")
+	bool IsCommandRegistered(FScriptingCommandInfo CommandInfo, bool bCheckInputChord = true) const;
+	
+	/** Retrieves the list of command info for all commands currently registered in the subsystem */
+	UFUNCTION(BlueprintCallable, Category = "Editor Scripting | Commands | Global")
+	TArray<FScriptingCommandInfo> GetRegisteredCommands() const;
+
+	/** Checks whether commands registered in the subsystem can be executed */
+	UFUNCTION(BlueprintCallable, Category = "Editor Scripting | Commands | Global")
+	bool CanExecuteCommands() const;
+
+	/** Sets whether commands registered in the subsystem can be executed */
+	UFUNCTION(BlueprintCallable, Category = "Editor Scripting | Commands | Global")
+	void SetCanExecuteCommands(bool bShouldExecuteCommands);
+
+	/**
+	 * Unregisters all commands dynamically registered within all contexts and sets.
+	 * @warning this will unregister all commands currently registered by this subsystem
+	 */
+	UFUNCTION(BlueprintCallable, Category = "Editor Scripting | Commands | Global")
+	void UnregisterAllSets();
+
+	
+	/**
+	 * Registers a new command set
+	 * @return Whether the set did not already exist and was successfully registered
+	 */
+	UFUNCTION(BlueprintCallable, Category = "Editor Scripting | Commands | Sets")
+	bool RegisterCommandSet(FName SetName);
+
+	/** Checks whether the given set is currently registered in the subsystem */
+	UFUNCTION(BlueprintCallable, Category = "Editor Scripting | Commands | Sets")
+	bool IsCommandSetRegistered(FName SetName) const;
+	
+	/**
+	 * Unregisters the corresponding command set with all commands registered within it
+	 * @return Whether the command set existed and was successfully unregistered
+	 */
+	UFUNCTION(BlueprintCallable, Category = "Editor Scripting | Commands | Sets")
+	bool UnregisterCommandSet(FName SetName);
+
+	/** Enables or disables execution of commands registered within the given set */
+	UFUNCTION(BlueprintCallable, Category = "Editor Scripting | Commands | Sets", Meta = (DisplayName = "Set Can Execute Commands"))
+	void SetCanSetExecuteCommands(FName SetName, bool bShouldExecuteCommands);
+
+	/** Checks whether commands in the given set can be executed. This will also check CanExecuteCommands at a global scope */
+	UFUNCTION(BlueprintCallable, Category = "Editor Scripting | Commands | Sets", Meta = (DisplayName = "Can Execute Commands"))
+	bool CanSetExecuteCommands(const FName SetName) const;
+
+	
+	/**
+	 * Retrieves the list of names for all contexts currently registered in the subsystem.
+	 * This does not check whether the contexts are bound to any UI Command List. 
+	 */
+	UFUNCTION(BlueprintCallable, Category = "Editor Scripting | Commands | Contexts")
+	TArray<FName> GetAvailableContexts() const;
+	
+	/**
+	 * Checks whether the context with the given name is currently registered in the subsystem
+	 * This does not check whether the context is bound to any UI Command List. 
+	 */
+	UFUNCTION(BlueprintCallable, Category = "Editor Scripting | Commands | Contexts")
+	bool IsContextRegistered(FName ContextName) const;
+	
+	/**
+	 * Retrieves the number of UI Command Lists registered within this context through the subsystem.
+	 * UI Command Lists are typically used to bind the list of commands associated with a single UI (i.e. a single viewport).
+	 */
+	UFUNCTION(BlueprintCallable, Category = "Editor Scripting | Commands | Contexts")
+	int GetBindingCountForContext(FName ContextName);
+
+	/**
+	 * Checks whether the given input chord is already mapped to a command in the given context.
+	 * This includes commands not registered through the subsystem.
+	 */
+	UFUNCTION(BlueprintCallable, Category = "Editor Scripting | Commands | Contexts")
+	bool IsInputChordMapped(FName ContextName, FInputChord InputChord) const;	
+	
+private:
+	/** The internal method to dynamically registers an editor command */
+	bool RegisterNewScriptingCommand(FScriptingCommandInfo CommandInfo, FExecuteCommand OnExecuteCommand,
+	                               FCanExecuteAction OnCanExecuteAction, bool bOverrideExisting = false);
+
+	/** Registers a context command in its command lists. This assumes the command is already contained in a set */
+	bool RegisterScriptingCommand(const TSharedRef<FScriptingCommand> ScriptingCommand);
+
+	/** Unregisters a context command from its command lists. This assumes the command is already contained in a set */
+	bool UnregisterScriptingCommand(const TSharedRef<FScriptingCommand> ScriptingCommand);
+
+	/** The delegate bound to registered commands. It will broadcast back to the given dynamic delegate */
+	UFUNCTION()
+	static void HandleExecuteAction(FExecuteCommand OnExecuteAction, const FScriptingCommandInfo CommandInfo);
+
+	/**
+	 * The delegate used to check whether a command can be executed.
+	 * @return true if the given delegate returns true and the given set and all subsystem commands are enabled
+	 */
+	UFUNCTION()
+	bool HandleCanExecuteAction(FCanExecuteCommand OnCanExecuteAction, const FScriptingCommandInfo CommandInfo) const;
+
+	/**
+	 * The default delegate to check if a command can be executed (when the user does not provide a specific delegate).
+	 * @return true if the given set and all subsystem commands are enabled
+	 */
+	UFUNCTION()
+	bool DefaultCanExecuteAction(FName SetName) const;
+
+private:
+	/** The list of commands contexts. Each context contains its own list of registered commands and UI command lists */
+	TMap<FName, FScriptingCommandsContext> CommandsInContext;
+	
+	/** The list of command sets. Each set contains its own list of registered commands */
+	TMap<FName, FScriptingCommandSet> CommandSets;
+
+	/** Whether the commands for all the sets of the subsystem can be executed */
+	bool bCanExecuteCommands = true;
+};
\ No newline at end of file
diff --git a/Engine/Plugins/Slate/SlateScripting/Source/SlateScriptingCommands/SlateScriptingCommands.Build.cs b/Engine/Plugins/Slate/SlateScripting/Source/SlateScriptingCommands/SlateScriptingCommands.Build.cs
new file mode 100644
--- /dev/null
+++ b/Engine/Plugins/Slate/SlateScripting/Source/SlateScriptingCommands/SlateScriptingCommands.Build.cs
@@ -0,0 +1,27 @@
+// Copyright Epic Games, Inc. All Rights Reserved.
+
+namespace UnrealBuildTool.Rules
+{
+	public class SlateScriptingCommands : ModuleRules
+	{
+		public SlateScriptingCommands(ReadOnlyTargetRules Target) : base(Target)
+		{
+			PublicDependencyModuleNames.AddRange(
+				new string[]
+				{
+					"Core",
+					"CoreUObject",
+					"Slate",
+					"SlateCore",
+					"Engine"
+				});
+			
+			PrivateDependencyModuleNames.AddRange(
+			new string[]
+			{
+				"InputCore"
+			});
+		}
+	}
+}
+
diff --git a/Engine/Source/Editor/ContentBrowser/Private/SContentBrowser.cpp b/Engine/Source/Editor/ContentBrowser/Private/SContentBrowser.cpp
--- a/Engine/Source/Editor/ContentBrowser/Private/SContentBrowser.cpp
+++ b/Engine/Source/Editor/ContentBrowser/Private/SContentBrowser.cpp
@@ -571,6 +571,8 @@ void SContentBrowser::BindCommands()
 			CommmandExtenderDelegates[i].Execute(Commands.ToSharedRef(), FOnContentBrowserGetSelection::CreateSP(this, &SContentBrowser::GetSelectionState));
 		}
 	}
+
+	FInputBindingManager::Get().RegisterCommandList(FContentBrowserCommands::Get().GetContextName(), Commands.ToSharedRef());
 }
 
 EVisibility SContentBrowser::GetFavoriteFolderVisibility() const
diff --git a/Engine/Source/Editor/LevelEditor/Private/LevelEditor.cpp b/Engine/Source/Editor/LevelEditor/Private/LevelEditor.cpp
--- a/Engine/Source/Editor/LevelEditor/Private/LevelEditor.cpp
+++ b/Engine/Source/Editor/LevelEditor/Private/LevelEditor.cpp
@@ -262,6 +262,9 @@ void FLevelEditorModule::StartupModule()
 	// Bind level editor commands shared across an instance
 	BindGlobalLevelEditorCommands();
 
+	// Exposes the global level editor command list to subscribers from other systems
+	FInputBindingManager::Get().RegisterCommandList(FLevelEditorCommands::Get().GetContextName(), GetGlobalLevelEditorActions());
+
 	FViewportTypeDefinition ViewportType = FViewportTypeDefinition([](const FAssetEditorViewportConstructionArgs& ConstructionArgs, TSharedPtr<ILevelEditor> InLevelEditor)
 		{
 			TSharedPtr<SLevelViewport> EditorViewport = SNew(SLevelViewport, ConstructionArgs)
@@ -327,7 +330,7 @@ void FLevelEditorModule::ShutdownModule()
 	{
 		FGlobalTabmanager::Get()->UnregisterTabSpawner("LevelEditor");
 		FModuleManager::GetModuleChecked<ISlateReflectorModule>("SlateReflector").UnregisterTabSpawner();
-	}	
+	}
 
 	FLevelEditorCommands::Unregister();
 	FLevelEditorModesCommands::Unregister();
@@ -1773,6 +1776,6 @@ void FLevelEditorModule::BindGlobalLevelEditorCommands()
 		FExecuteAction::CreateStatic(&FLevelEditorActionCallbacks::OpenMergeActor_Clicked)
 		);
 }
-	
+
 
 #undef LOCTEXT_NAMESPACE
diff --git a/Engine/Source/Editor/LevelEditor/Private/SLevelViewport.cpp b/Engine/Source/Editor/LevelEditor/Private/SLevelViewport.cpp
--- a/Engine/Source/Editor/LevelEditor/Private/SLevelViewport.cpp
+++ b/Engine/Source/Editor/LevelEditor/Private/SLevelViewport.cpp
@@ -1324,6 +1324,9 @@ void SLevelViewport::BindCommands()
 	}
 
 	UICommandListRef.SetCanProduceActionForCommand( FUICommandList::FCanProduceActionForCommand::CreateSP(this, &SLevelViewport::CanProduceActionForCommand) );
+
+	// Exposes the current level viewport command list to subscribers from other systems
+	FInputBindingManager::Get().RegisterCommandList(FLevelViewportCommands::Get().GetContextName(), CommandList.ToSharedRef());
 }
 	
 void SLevelViewport::BindOptionCommands( FUICommandList& OutCommandList )
diff --git a/Engine/Source/Editor/MainFrame/Private/MainFrameModule.cpp b/Engine/Source/Editor/MainFrame/Private/MainFrameModule.cpp
--- a/Engine/Source/Editor/MainFrame/Private/MainFrameModule.cpp
+++ b/Engine/Source/Editor/MainFrame/Private/MainFrameModule.cpp
@@ -540,6 +540,9 @@ void FMainFrameModule::StartupModule( )
 	FGenericCommands::Register();
 	FMainFrameCommands::Register();
 
+	// Exposes the main frame command list to subscribers from other systems
+	FInputBindingManager::Get().RegisterCommandList(FMainFrameCommands::Get().GetContextName(), FMainFrameCommands::ActionList);
+
 	SetLevelNameForWindowTitle(TEXT(""));
 
 	// Register to find out about when hot reload completes, so we can show a notification
@@ -613,6 +616,7 @@ void FMainFrameModule::HandleResizeMainFrameCommand(const TArray<FString>& Args)
 	
 }
 
+
 /* FMainFrameModule implementation
  *****************************************************************************/
 
diff --git a/Engine/Source/Editor/Sequencer/Private/Sequencer.cpp b/Engine/Source/Editor/Sequencer/Private/Sequencer.cpp
--- a/Engine/Source/Editor/Sequencer/Private/Sequencer.cpp
+++ b/Engine/Source/Editor/Sequencer/Private/Sequencer.cpp
@@ -146,6 +146,7 @@
 #include "SSequencerGroupManager.h"
 #include "ActorTreeItem.h"
 #include "Widgets/Layout/SSpacer.h"
+#include "CurveEditorCommands.h"
 
 #include "EntitySystem/MovieSceneEntitySystemLinker.h"
 #include "EntitySystem/MovieScenePreAnimatedStateSystem.h"
@@ -570,6 +571,10 @@ FSequencer::FSequencer()
 	Selection.GetOnOutlinerNodeSelectionChanged().AddRaw(this, &FSequencer::OnSelectedOutlinerNodesChanged);
 	Selection.GetOnNodesWithSelectedKeysOrSectionsChanged().AddRaw(this, &FSequencer::OnSelectedOutlinerNodesChanged);
 	Selection.GetOnOutlinerNodeSelectionChangedObjectGuids().AddRaw(this, &FSequencer::OnSelectedOutlinerNodesChanged);
+
+	// Exposes the sequencer and curve editor command lists to subscribers from other systems
+	FInputBindingManager::Get().RegisterCommandList(FSequencerCommands::Get().GetContextName(), SequencerCommandBindings);
+	FInputBindingManager::Get().RegisterCommandList(FCurveEditorCommands::Get().GetContextName(), CurveEditorSharedBindings);
 }
 
 
diff --git a/Engine/Source/Runtime/Slate/Private/Framework/Commands/InputBindingManager.cpp b/Engine/Source/Runtime/Slate/Private/Framework/Commands/InputBindingManager.cpp
--- a/Engine/Source/Runtime/Slate/Private/Framework/Commands/InputBindingManager.cpp
+++ b/Engine/Source/Runtime/Slate/Private/Framework/Commands/InputBindingManager.cpp
@@ -553,6 +553,33 @@ bool FInputBindingManager::CommandPassesFilter(const FName InBindingContext, con
 	return true;
 }
 
+bool FInputBindingManager::RegisterCommandList(const FName InBindingContext, const TSharedRef<FUICommandList> CommandList) const
+{
+	if (ContextMap.Contains(InBindingContext) && OnRegisterCommandList.IsBound())
+	{
+		OnRegisterCommandList.Broadcast(InBindingContext, CommandList);
+		return true;
+	}
+	return false;
+}
+
+TSharedPtr<FUICommandList> FInputBindingManager::RegisterNewCommandList(const FName InBindingContext) const
+{
+	const TSharedPtr<FUICommandList> CommandList = MakeShared<FUICommandList>();
+
+	return RegisterCommandList(InBindingContext, CommandList.ToSharedRef()) ? CommandList : nullptr;
+}
+
+bool FInputBindingManager::UnregisterCommandList(const FName InBindingContext, TSharedRef<FUICommandList> CommandList) const
+{
+	if (ContextMap.Contains(InBindingContext) && OnUnregisterCommandList.IsBound())
+	{
+		OnUnregisterCommandList.Broadcast(InBindingContext, CommandList);
+		return true;
+	}
+	return false;
+}
+
 void FInputBindingManager::AddCommandFilter(const FName InOwnerName, const FName InBindingContext, const FName InCommandName, const ECommandFilterType FilterType)
 {
 	if (FilterType == ECommandFilterType::Blacklist)
diff --git a/Engine/Source/Runtime/Slate/Public/Framework/Commands/InputBindingManager.h b/Engine/Source/Runtime/Slate/Public/Framework/Commands/InputBindingManager.h
--- a/Engine/Source/Runtime/Slate/Public/Framework/Commands/InputBindingManager.h
+++ b/Engine/Source/Runtime/Slate/Public/Framework/Commands/InputBindingManager.h
@@ -3,6 +3,7 @@
 #pragma once
 
 #include "CoreMinimal.h"
+#include "UICommandList.h"
 #include "Framework/Commands/InputChord.h"
 #include "Framework/Commands/UICommandInfo.h"
 
@@ -15,6 +16,9 @@ typedef TMap<FInputChord, FName> FChordMap;
 /** Delegate for alerting subscribers the input manager records a user-defined chord */
 DECLARE_MULTICAST_DELEGATE_OneParam(FOnUserDefinedChordChanged, const FUICommandInfo& );
 
+/** Delegate for exposing new command lists to subscribers */
+DECLARE_MULTICAST_DELEGATE_TwoParams(FOnRegisterCommandList, const FName, TSharedRef<FUICommandList>);
+
 /**
  * Manager responsible for creating and processing input bindings.                  
  */
@@ -171,6 +175,36 @@ public:
 	*/
 	bool CommandPassesFilter(const FName InBindingContext, const FName InCommandName) const;
 
+	/**
+	 * Exposes the given command list and its context to all OnRegisterCommandList subscribers
+	 *
+	 * @param InBindingContext	The context in which the command list is active
+	 * @param CommandList		The command list to expose to subscribers
+	 */
+	bool RegisterCommandList(const ::FName InBindingContext, TSharedRef<FUICommandList> CommandList) const;
+
+	/**
+	 * Exposes a new given command list from the given context to all OnRegisterCommandList subscribers
+	 * 
+	 * @param InBindingContext	The context in which the command list is active
+	 * @return	The new command list exposed to subscribers
+	 */
+	TSharedPtr<FUICommandList> RegisterNewCommandList(const FName InBindingContext) const;
+
+	/**
+	 * Exposes the given command list and its context to all OnUnregisterCommandList subscribers
+	 * 
+	 * @param InBindingContext	The context in which the command list is active
+	 * @param CommandList		The command list to expose to subscribers
+	 */
+	bool UnregisterCommandList(const ::FName InBindingContext, TSharedRef<FUICommandList> CommandList) const;
+
+	/** A delegate for systems to subscribe to any other system exposing their command lists */
+	FOnRegisterCommandList OnRegisterCommandList;
+
+	/** A delegate for systems to subscribe to any other system unregistering their command lists */
+	FOnRegisterCommandList OnUnregisterCommandList;
+
 private:
 
 	/**
