Fix SlateSriptingCommand module - Compilation error in CL 17735741

diff --git a/Engine/Plugins/Slate/SlateScripting/Source/SlateScriptingCommands/Private/SlateScriptingCommandsModule.cpp b/Engine/Plugins/Slate/SlateScripting/Source/SlateScriptingCommands/Private/SlateScriptingCommandsModule.cpp
--- a/Engine/Plugins/Slate/SlateScripting/Source/SlateScriptingCommands/Private/SlateScriptingCommandsModule.cpp
+++ b/Engine/Plugins/Slate/SlateScripting/Source/SlateScriptingCommands/Private/SlateScriptingCommandsModule.cpp
@@ -10,4 +10,4 @@ void FSlateScriptingCommandsModule::ShutdownModule()
 {
 }
 
-IMPLEMENT_MODULE(FSlateScriptingCommandsModule, SlateScriptingCommands);
\ No newline at end of file
+IMPLEMENT_MODULE(FSlateScriptingCommandsModule, SlateScriptingCommands)
\ No newline at end of file
diff --git a/Engine/Plugins/Slate/SlateScripting/Source/SlateScriptingCommands/Public/SlateScriptingCommandsModule.h b/Engine/Plugins/Slate/SlateScripting/Source/SlateScriptingCommands/Public/SlateScriptingCommandsModule.h
--- a/Engine/Plugins/Slate/SlateScripting/Source/SlateScriptingCommands/Public/SlateScriptingCommandsModule.h
+++ b/Engine/Plugins/Slate/SlateScripting/Source/SlateScriptingCommands/Public/SlateScriptingCommandsModule.h
@@ -3,7 +3,7 @@
 #pragma once
 
 #include "CoreMinimal.h"
-#include "Modules/ModuleInterface.h"
+#include "Modules/ModuleManager.h"
 
 class FSlateScriptingCommandsModule : public IModuleInterface
 {
diff --git a/Engine/Plugins/Slate/SlateScripting/Source/SlateScriptingCommands/SlateScriptingCommands.Build.cs b/Engine/Plugins/Slate/SlateScripting/Source/SlateScriptingCommands/SlateScriptingCommands.Build.cs
--- a/Engine/Plugins/Slate/SlateScripting/Source/SlateScriptingCommands/SlateScriptingCommands.Build.cs
+++ b/Engine/Plugins/Slate/SlateScripting/Source/SlateScriptingCommands/SlateScriptingCommands.Build.cs
@@ -6,6 +6,8 @@ namespace UnrealBuildTool.Rules
 	{
 		public SlateScriptingCommands(ReadOnlyTargetRules Target) : base(Target)
 		{
+			PCHUsage = PCHUsageMode.UseExplicitOrSharedPCHs;
+
 			PublicDependencyModuleNames.AddRange(
 				new string[]
 				{
