Typo fixes in comments

diff --git a/Engine/Plugins/Slate/SlateScripting/Source/SlateScriptingCommands/Public/UICommandsScriptingSubsystem.h b/Engine/Plugins/Slate/SlateScripting/Source/SlateScriptingCommands/Public/UICommandsScriptingSubsystem.h
--- a/Engine/Plugins/Slate/SlateScripting/Source/SlateScriptingCommands/Public/UICommandsScriptingSubsystem.h
+++ b/Engine/Plugins/Slate/SlateScripting/Source/SlateScriptingCommands/Public/UICommandsScriptingSubsystem.h
@@ -196,7 +196,7 @@ public:
 	 * Registers a command list in the given context. Optionally registering the context itself if it does not exist
 	 * Note the context should already be existing and registered through the binding manager.
 	 * 
-	 * This should be called by modules and tools looking to expose there command lists to the subsystem
+	 * This should be called by modules and tools looking to expose their command lists to the subsystem.
 	 * 
 	 * @param	ContextName		The context name as defined in the binding manager
 	 * @param	CommandList		The command list to register in this context
diff --git a/Engine/Source/Runtime/Engine/Private/Subsystems/SubsystemCollection.cpp b/Engine/Source/Runtime/Engine/Private/Subsystems/SubsystemCollection.cpp
--- a/Engine/Source/Runtime/Engine/Private/Subsystems/SubsystemCollection.cpp
+++ b/Engine/Source/Runtime/Engine/Private/Subsystems/SubsystemCollection.cpp
@@ -236,7 +236,7 @@ USubsystem* FSubsystemCollectionBase::AddAndInitializeSubsystem(UClass* Subsyste
 			// Catch any attempt to add a subsystem of the wrong type
 			checkf(SubsystemClass->IsChildOf(BaseType), TEXT("ClassType (%s) must be a subclass of BaseType(%s)."), *SubsystemClass->GetName(), *BaseType->GetName());
 
-			// Do not create instances of classes aren't authoritative
+			// Do not create instances of classes that aren't authoritative.
 			if (SubsystemClass->GetAuthoritativeClass() != SubsystemClass)
 			{	
 				return nullptr;
